- Renamed UnrealDisasterRecoveryService as UnrealRecoverySvc
  - Set a ConcertSyncServer ShortName as "CncrtSyncSvr" to ensure shorter build path.

The change saves 29 characters on the offending path. The path before vs the path after:

Engine\Plugins\Developer\Concert\ConcertSync\ConcertSyncServer\Intermediate\Build\Win64\UnrealDisasterRecoveryService\Development\ConcertSyncServer\UnrealDisasterRecoveryService-ConcertSyncServer.lib (Before, 199 chars)
Engine\Plugins\Developer\Concert\ConcertSync\ConcertSyncServer\Intermediate\Build\Win64\UnrealRecoverySvc\Development\CncrtSyncSvr\UnrealRecoverySvc-ConcertSyncServer.lib (After, 170 chars)

Edigrated 11281991 from Dev-VirtualProduction

diff --git a/Engine/Build/InstalledEngineBuild.xml b/Engine/Build/InstalledEngineBuild.xml
--- a/Engine/Build/InstalledEngineBuild.xml
+++ b/Engine/Build/InstalledEngineBuild.xml
@@ -477,7 +477,7 @@
 			<Compile Target="UnrealLightmass" Configuration="Development" Platform="Win64" Tag="#Build Tools Win64"/>
 			<Compile Target="UnrealPak" Configuration="Development" Platform="Win64" Tag="#Build Tools Win64"/>
 			<Compile Target="UnrealMultiUserServer" Configuration="Development" Platform="Win64" Tag="#Build Tools Win64"/>
-			<Compile Target="UnrealDisasterRecoveryService" Configuration="Development" Platform="Win64" Tag="#Build Tools Win64"/>
+			<Compile Target="UnrealRecoverySvc" Configuration="Development" Platform="Win64" Tag="#Build Tools Win64"/>
 			<Compile Target="LiveCodingConsole" Configuration="Development" Platform="Win64" Tag="#Build Tools Win64"/>
 			<Compile Target="BootstrapPackagedGame" Configuration="Shipping" Platform="Win64" Tag="#Build Tools Win64"/>
 			<Compile Target="BuildPatchTool" Configuration="Shipping" Platform="Win64" Tag="#Build Tools Win64"/>
@@ -545,7 +545,7 @@
 			<Compile Target="UnrealLightmass" Configuration="Development" Platform="Linux" Tag="#Build Tools Linux"/>
 			<Compile Target="UnrealPak" Configuration="Development" Platform="Linux" Tag="#Build Tools Linux"/>
 			<Compile Target="UnrealMultiUserServer" Configuration="Development" Platform="Linux" Tag="#Build Tools Linux"/>
-			<Compile Target="UnrealDisasterRecoveryService" Configuration="Development" Platform="Linux" Tag="#Build Tools Linux"/>
+			<Compile Target="UnrealRecoverySvc" Configuration="Development" Platform="Linux" Tag="#Build Tools Linux"/>
 		</Node>
 		<Node Name="Build Tools Win32" Requires="Compile UnrealHeaderTool Win64">
 			<Compile Target="CrashReportClient" Configuration="Shipping" Platform="Win32" Tag="#Build Tools Win32"/>
@@ -565,7 +565,7 @@
 			<Compile Target="UnrealLightmass" Configuration="Development" Platform="Mac" Tag="#Build Tools Mac"/>
 			<Compile Target="UnrealPak" Configuration="Development" Platform="Mac" Tag="#Build Tools Mac"/>
 			<Compile Target="UnrealMultiUserServer" Configuration="Development" Platform="Mac" Tag="#Build Tools Mac"/>
-			<Compile Target="UnrealDisasterRecoveryService" Configuration="Development" Platform="Mac" Tag="#Build Tools Mac"/>
+			<Compile Target="UnrealRecoverySvc" Configuration="Development" Platform="Mac" Tag="#Build Tools Mac"/>
 			<Compile Target="BuildPatchTool" Configuration="Shipping" Platform="Mac" Tag="#Build Tools Mac"/>
 			<Compile Target="DsymExporter" Configuration="Development" Platform="Mac" Tag="#Build Tools Mac"/>
 			<Compile Target="UE4EditorServices" Configuration="Development" Platform="Mac" Tag="#Build Tools Mac"/>
diff --git a/Engine/Plugins/Developer/Concert/ConcertMain/ConcertMain.uplugin b/Engine/Plugins/Developer/Concert/ConcertMain/ConcertMain.uplugin
--- a/Engine/Plugins/Developer/Concert/ConcertMain/ConcertMain.uplugin
+++ b/Engine/Plugins/Developer/Concert/ConcertMain/ConcertMain.uplugin
@@ -15,19 +15,19 @@
 	"Hidden": true,
 	"IsBetaVersion": true,
 	"Installed": false,
-	"SupportedPrograms": [ "UnrealMultiUserServer", "UnrealDisasterRecoveryService", "CrashReportClientEditor" ],
+	"SupportedPrograms": [ "UnrealMultiUserServer", "UnrealRecoverySvc", "CrashReportClientEditor" ],
 	"Modules": [
 		{
 			"Name": "Concert",
 			"Type": "UncookedOnly",
 			"LoadingPhase": "PreDefault",
-			"WhitelistPrograms": [ "UnrealMultiUserServer", "UnrealDisasterRecoveryService", "CrashReportClientEditor" ]
+			"WhitelistPrograms": [ "UnrealMultiUserServer", "UnrealRecoverySvc", "CrashReportClientEditor" ]
 		},
 		{
 			"Name": "ConcertTransport",
 			"Type": "UncookedOnly",
 			"LoadingPhase": "PreDefault",
-			"WhitelistPrograms": [ "UnrealMultiUserServer", "UnrealDisasterRecoveryService", "CrashReportClientEditor" ]
+			"WhitelistPrograms": [ "UnrealMultiUserServer", "UnrealRecoverySvc", "CrashReportClientEditor" ]
 		}
 	]
 }
diff --git a/Engine/Plugins/Developer/Concert/ConcertSync/ConcertSyncCore/ConcertSyncCore.uplugin b/Engine/Plugins/Developer/Concert/ConcertSync/ConcertSyncCore/ConcertSyncCore.uplugin
--- a/Engine/Plugins/Developer/Concert/ConcertSync/ConcertSyncCore/ConcertSyncCore.uplugin
+++ b/Engine/Plugins/Developer/Concert/ConcertSync/ConcertSyncCore/ConcertSyncCore.uplugin
@@ -15,13 +15,13 @@
 	"Hidden": true,
 	"IsBetaVersion": true,
 	"Installed": false,
-	"SupportedPrograms": [ "UnrealMultiUserServer", "UnrealDisasterRecoveryService", "CrashReportClientEditor" ],
+	"SupportedPrograms": [ "UnrealMultiUserServer", "UnrealRecoverySvc", "CrashReportClientEditor" ],
 	"Modules": [
 		{
 			"Name": "ConcertSyncCore",
 			"Type": "UncookedOnly",
 			"LoadingPhase": "PreDefault",
-			"WhitelistPrograms": [ "UnrealMultiUserServer", "UnrealDisasterRecoveryService", "CrashReportClientEditor" ]
+			"WhitelistPrograms": [ "UnrealMultiUserServer", "UnrealRecoverySvc", "CrashReportClientEditor" ]
 		}
 	],
 	"Plugins": [
diff --git a/Engine/Plugins/Developer/Concert/ConcertSync/ConcertSyncServer/ConcertSyncServer.uplugin b/Engine/Plugins/Developer/Concert/ConcertSync/ConcertSyncServer/ConcertSyncServer.uplugin
--- a/Engine/Plugins/Developer/Concert/ConcertSync/ConcertSyncServer/ConcertSyncServer.uplugin
+++ b/Engine/Plugins/Developer/Concert/ConcertSync/ConcertSyncServer/ConcertSyncServer.uplugin
@@ -15,13 +15,13 @@
 	"Hidden": true,
 	"IsBetaVersion": true,
 	"Installed": false,
-	"SupportedPrograms": [ "UnrealMultiUserServer", "UnrealDisasterRecoveryService", "CrashReportClientEditor" ],
+	"SupportedPrograms": [ "UnrealMultiUserServer", "UnrealRecoverySvc", "CrashReportClientEditor" ],
 	"Modules": [
 		{
 			"Name": "ConcertSyncServer",
 			"Type": "UncookedOnly",
 			"LoadingPhase": "Default",
-			"WhitelistPrograms": [ "UnrealMultiUserServer", "UnrealDisasterRecoveryService", "CrashReportClientEditor" ]
+			"WhitelistPrograms": [ "UnrealMultiUserServer", "UnrealRecoverySvc", "CrashReportClientEditor" ]
 		}
 	],
 	"Plugins": [
diff --git a/Engine/Plugins/Developer/Concert/ConcertSync/ConcertSyncServer/Source/ConcertSyncServer/ConcertSyncServer.Build.cs b/Engine/Plugins/Developer/Concert/ConcertSync/ConcertSyncServer/Source/ConcertSyncServer/ConcertSyncServer.Build.cs
--- a/Engine/Plugins/Developer/Concert/ConcertSync/ConcertSyncServer/Source/ConcertSyncServer/ConcertSyncServer.Build.cs
+++ b/Engine/Plugins/Developer/Concert/ConcertSync/ConcertSyncServer/Source/ConcertSyncServer/ConcertSyncServer.Build.cs
@@ -6,6 +6,9 @@ namespace UnrealBuildTool.Rules
 	{
 		public ConcertSyncServer(ReadOnlyTargetRules Target) : base(Target)
 		{
+			// Make the linking path shorter (to prevent busting the Windows limit) when linking ConcertSyncServer.lib against an executable that have a long name.
+			ShortName = "CncrtSyncSvr";
+
 			PublicDependencyModuleNames.AddRange(
 				new string[]
 				{
diff --git a/Engine/Plugins/Developer/Concert/DisasterRecoveryClient/Source/DisasterRecoveryClient/Private/DisasterRecoveryClientModule.cpp b/Engine/Plugins/Developer/Concert/DisasterRecoveryClient/Source/DisasterRecoveryClient/Private/DisasterRecoveryClientModule.cpp
--- a/Engine/Plugins/Developer/Concert/DisasterRecoveryClient/Source/DisasterRecoveryClient/Private/DisasterRecoveryClientModule.cpp
+++ b/Engine/Plugins/Developer/Concert/DisasterRecoveryClient/Source/DisasterRecoveryClient/Private/DisasterRecoveryClientModule.cpp
@@ -48,7 +48,7 @@ namespace DisasterRecoveryUtil
 static const FName RecoveryHubTabName("RecoveryHub");
 
 /**
- * Return the name of the executable hosting disaster recovery service, like 'UnrealDisasterRecoveryService' without the extension.
+ * Return the name of the executable hosting disaster recovery service, like 'UnrealRecoverySvc' without the extension.
  */
 FString GetDisasterRecoveryServiceExeName()
 {
@@ -58,7 +58,7 @@ FString GetDisasterRecoveryServiceExeName()
 	}
 	else
 	{
-		return TEXT("UnrealDisasterRecoveryService");
+		return TEXT("UnrealRecoverySvc");
 	}
 }
 
@@ -327,7 +327,7 @@ private:
 		const FString DisasterRecoveryServicePath = GetDisasterRecoveryServicePath();
 		if (DisasterRecoveryServicePath.IsEmpty())
 		{
-			UE_LOG(LogDisasterRecovery, Warning, TEXT("Disaster Recovery Service application was not found. Disaster Recovery will be disabled! Please build 'UnrealDisasterRecoveryService'."));
+			UE_LOG(LogDisasterRecovery, Warning, TEXT("Disaster Recovery Service application was not found. Disaster Recovery will be disabled! Please build '%s'."), *DisasterRecoveryUtil::GetDisasterRecoveryServiceExeName());
 			return false;
 		}
 
@@ -352,7 +352,7 @@ private:
 		// Started for the first time (not restarted)? Always create the manager (and start the service) even if disaster recovery is disabled (in the settings). This allows the user to import sessions for crash analysis.
 		if (!SessionManager)
 		{
-			// If crash reporter is running out of process, it also hosts disaster recovery server as the '-ConcertServer' param is set when spawning CrashReporterClientEditor. No need to start the UnrealDisasterRecoveryService executable.
+			// If crash reporter is running out of process, it also hosts disaster recovery server as the '-ConcertServer' param is set when spawning CrashReporterClientEditor. No need to start the UnrealRecoverySvc executable.
 			if (!FGenericCrashContext::IsOutOfProcessCrashReporter() && !SpawnDisasterRecoveryServer(RecoveryService::GetRecoveryServerName()))
 			{
 				return; // Failed to spawn the service.
diff --git a/Engine/Plugins/Messaging/UdpMessaging/UdpMessaging.uplugin b/Engine/Plugins/Messaging/UdpMessaging/UdpMessaging.uplugin
--- a/Engine/Plugins/Messaging/UdpMessaging/UdpMessaging.uplugin
+++ b/Engine/Plugins/Messaging/UdpMessaging/UdpMessaging.uplugin
@@ -18,7 +18,7 @@
 	[
 		"UnrealFrontend",
 		"UnrealMultiUserServer",
-		"UnrealDisasterRecoveryService",
+		"UnrealRecoverySvc",
 		"UnrealLightmass",
 		"CrashReportClientEditor"
 	],
@@ -34,7 +34,7 @@
 			"WhitelistPrograms": [
 				"UnrealFrontend",
 				"UnrealMultiUserServer",
-				"UnrealDisasterRecoveryService",
+				"UnrealRecoverySvc",
 				"UnrealLightmass",
 				"CrashReportClientEditor"
 			]
diff --git a/Engine/Plugins/Runtime/Database/SQLiteCore/SQLiteCore.uplugin b/Engine/Plugins/Runtime/Database/SQLiteCore/SQLiteCore.uplugin
--- a/Engine/Plugins/Runtime/Database/SQLiteCore/SQLiteCore.uplugin
+++ b/Engine/Plugins/Runtime/Database/SQLiteCore/SQLiteCore.uplugin
@@ -14,7 +14,7 @@
 	"CanContainContent": false,
 	"IsBetaVersion": false,
 	"Installed": false,
-	"SupportedPrograms": [ "UnrealMultiUserServer", "UnrealDisasterRecoveryService" ],
+	"SupportedPrograms": [ "UnrealMultiUserServer", "UnrealRecoverySvc" ],
 	"Modules": [
 		{
 			"Name": "SQLiteCore",
diff --git a/Engine/Programs/UnrealDisasterRecoveryService/Config/DefaultEngine.ini b/Engine/Programs/UnrealDisasterRecoveryService/Config/DefaultEngine.ini
deleted file mode 100644
diff --git a/Engine/Programs/UnrealRecoverySvc/Config/DefaultEngine.ini b/Engine/Programs/UnrealRecoverySvc/Config/DefaultEngine.ini
new file mode 100644
--- /dev/null
+++ b/Engine/Programs/UnrealRecoverySvc/Config/DefaultEngine.ini
@@ -0,0 +1,10 @@
+[URL]
+GameNameShort=UDRS
+
+[Plugins]
++ProgramEnabledPlugins="UdpMessaging"
++ProgramEnabledPlugins="ConcertSyncServer"
+
+[/Script/UdpMessaging.UdpMessagingSettings]
+EnableTransport=True
+MulticastTimeToLive=0
diff --git a/Engine/Source/Programs/UnrealDisasterRecoveryService/Private/UnrealDisasterRecoveryService.cpp b/Engine/Source/Programs/UnrealDisasterRecoveryService/Private/UnrealDisasterRecoveryService.cpp
deleted file mode 100644
diff --git a/Engine/Source/Programs/UnrealDisasterRecoveryService/UnrealDisasterRecoveryService.Build.cs b/Engine/Source/Programs/UnrealDisasterRecoveryService/UnrealDisasterRecoveryService.Build.cs
deleted file mode 100644
diff --git a/Engine/Source/Programs/UnrealDisasterRecoveryService/UnrealDisasterRecoveryService.Target.cs b/Engine/Source/Programs/UnrealDisasterRecoveryService/UnrealDisasterRecoveryService.Target.cs
deleted file mode 100644
diff --git a/Engine/Source/Programs/UnrealRecoverySvc/Private/UnrealRecoverySvc.cpp b/Engine/Source/Programs/UnrealRecoverySvc/Private/UnrealRecoverySvc.cpp
new file mode 100644
--- /dev/null
+++ b/Engine/Source/Programs/UnrealRecoverySvc/Private/UnrealRecoverySvc.cpp
@@ -0,0 +1,60 @@
+// Copyright Epic Games, Inc. All Rights Reserved.
+
+#include "ConcertSettings.h"
+#include "ConcertSyncServerLoop.h"
+
+#include "RequiredProgramMainCPPInclude.h"
+
+IMPLEMENT_APPLICATION(UnrealRecoverySvc, "UnrealRecoverySvc");
+
+/**
+ * Application entry point
+ *
+ * @param	ArgC	Command-line argument count
+ * @param	ArgV	Argument strings
+ */
+INT32_MAIN_INT32_ARGC_TCHAR_ARGV()
+{
+	uint32 EditorProcessId = 0;
+
+	FConcertSyncServerLoopInitArgs ServerLoopInitArgs;
+	ServerLoopInitArgs.IdealFramerate = 30;
+	ServerLoopInitArgs.SessionFlags = EConcertSyncSessionFlags::Default_DisasterRecoverySession;
+	ServerLoopInitArgs.ServiceRole = TEXT("DisasterRecovery");
+	ServerLoopInitArgs.ServiceFriendlyName = TEXT("Disaster Recovery Service");
+	ServerLoopInitArgs.ServiceAutoArchiveSessionFilter.bIncludeIgnoredActivities = true;
+	ServerLoopInitArgs.bShowConsole = false;
+
+	ServerLoopInitArgs.GetServerConfigFunc = [&EditorProcessId]() -> const UConcertServerConfig*
+	{
+		FParse::Value(FCommandLine::Get(), TEXT("-EDITORPID="), EditorProcessId);
+		if (EditorProcessId)
+		{
+			UE_LOG(LogSyncServer, Display, TEXT("Watching Editor process %d"), EditorProcessId);
+		}
+		else
+		{
+			UE_LOG(LogSyncServer, Error, TEXT("Invalid -EditorPID argument. Cannot continue!"));
+			return nullptr;
+		}
+
+		UConcertServerConfig* ServerConfig = IConcertSyncServerModule::Get().ParseServerSettings(FCommandLine::Get());
+		ServerConfig->bAutoArchiveOnReboot = true; // If server crashed, was killed, etc, ensure the recovery session is archived (expected by recovery flow).
+		ServerConfig->EndpointSettings.RemoteEndpointTimeoutSeconds = 4; // Non-zero -> Ensure the Endpoints are re-registered with Message Bus if the UDP layer restarts or some errors trigger a 'auto-repair'.
+		ServerConfig->bMountDefaultSessionRepository = false; // Let the client mount its own repository to support concurrent service and avoid them to access the same non-sharable database files.
+		ServerConfig->AuthorizedClientKeys.Add(ServerConfig->ServerName); // The disaster recovery client is configured to use the unique server name as key to identify itself.
+		return ServerConfig;
+	};
+
+	FTicker::GetCoreTicker().AddTicker(TEXT("CheckEditorHealth"), 1.0f, [&EditorProcessId](float)
+	{
+		if (!FPlatformProcess::IsApplicationRunning(EditorProcessId))
+		{
+			UE_LOG(LogSyncServer, Warning, TEXT("Editor process %d lost! Requesting exit."), EditorProcessId);
+			FPlatformMisc::RequestExit(false);
+		}
+		return true;
+	});
+
+	return ConcertSyncServerLoop(ArgC, ArgV, ServerLoopInitArgs);
+}
diff --git a/Engine/Source/Programs/UnrealRecoverySvc/UnrealRecoverySvc.Build.cs b/Engine/Source/Programs/UnrealRecoverySvc/UnrealRecoverySvc.Build.cs
new file mode 100644
--- /dev/null
+++ b/Engine/Source/Programs/UnrealRecoverySvc/UnrealRecoverySvc.Build.cs
@@ -0,0 +1,39 @@
+// Copyright Epic Games, Inc. All Rights Reserved.
+
+using UnrealBuildTool;
+
+public class UnrealRecoverySvc : ModuleRules
+{
+	public UnrealRecoverySvc(ReadOnlyTargetRules Target) : base(Target)
+	{
+		PrivateDependencyModuleNames.AddRange(
+			new string[] {
+				"Core",
+				"CoreUObject",
+				"Concert",
+				"ApplicationCore",					// for LaunchEngineLoop.cpp dependency
+				"Projects",							// for LaunchEngineLoop.cpp dependency
+			}
+		);
+
+		PrivateIncludePathModuleNames.AddRange(
+			new string[] {
+				"Launch",
+				"ConcertSyncCore",
+				"ConcertSyncServer",
+			}
+		);
+
+		PrivateIncludePaths.AddRange(
+			new string[] {
+				"Runtime/Launch/Private",           // for LaunchEngineLoop.cpp include
+			}
+		);
+
+		DynamicallyLoadedModuleNames.AddRange(
+			new string[] {
+				"ConcertSyncServer",
+			}
+		);
+	}
+}
diff --git a/Engine/Source/Programs/UnrealRecoverySvc/UnrealRecoverySvc.Target.cs b/Engine/Source/Programs/UnrealRecoverySvc/UnrealRecoverySvc.Target.cs
new file mode 100644
--- /dev/null
+++ b/Engine/Source/Programs/UnrealRecoverySvc/UnrealRecoverySvc.Target.cs
@@ -0,0 +1,28 @@
+// Copyright Epic Games, Inc. All Rights Reserved.
+
+using UnrealBuildTool;
+using System.Collections.Generic;
+
+public class UnrealRecoverySvcTarget : TargetRules
+{
+	public UnrealRecoverySvcTarget(TargetInfo Target) : base(Target)
+	{
+		Type = TargetType.Program;
+		LinkType = TargetLinkType.Modular;
+		LaunchModuleName = "UnrealRecoverySvc";
+		AdditionalPlugins.Add("UdpMessaging");
+		AdditionalPlugins.Add("ConcertSyncServer");
+
+		// This app compiles against Core/CoreUObject, but not the Engine or Editor, so compile out Engine and Editor references from Core/CoreUObject
+		bCompileAgainstCoreUObject = true;
+		bCompileAgainstEngine = false;
+		bBuildWithEditorOnlyData = false;
+
+		// Enable Developer plugins (like Concert!)
+		bCompileWithPluginSupport = true;
+		bBuildDeveloperTools = true;
+
+		// This app is a console application (sets entry point to main(), instead of WinMain())
+		bIsBuildingConsoleApplication = true;
+	}
+}
