Fix horde issue 251250

Add ConcertServer module to UnrealRecoverySvc

diff --git a/Engine/Source/Programs/UnrealRecoverySvc/UnrealRecoverySvc.Build.cs b/Engine/Source/Programs/UnrealRecoverySvc/UnrealRecoverySvc.Build.cs
--- a/Engine/Source/Programs/UnrealRecoverySvc/UnrealRecoverySvc.Build.cs
+++ b/Engine/Source/Programs/UnrealRecoverySvc/UnrealRecoverySvc.Build.cs
@@ -11,6 +11,7 @@ public class UnrealRecoverySvc : ModuleRules
 				"Core",
 				"CoreUObject",
 				"Concert",
+				"ConcertServer",
 				"ConcertSyncCore",
 				"ApplicationCore",					// for LaunchEngineLoop.cpp dependency
 				"Projects",							// for LaunchEngineLoop.cpp dependency
